Fix undefined variable in TextToSpeech
#rnx

diff --git a/Engine/Plugins/Experimental/TextToSpeech/Source/TextToSpeech/Private/Flite/FliteAdapter.cpp b/Engine/Plugins/Experimental/TextToSpeech/Source/TextToSpeech/Private/Flite/FliteAdapter.cpp
--- a/Engine/Plugins/Experimental/TextToSpeech/Source/TextToSpeech/Private/Flite/FliteAdapter.cpp
+++ b/Engine/Plugins/Experimental/TextToSpeech/Source/TextToSpeech/Private/Flite/FliteAdapter.cpp
@@ -30,12 +30,13 @@ extern "C"
 #include "Windows/HideWindowsPlatformTypes.h"
 #endif
 
-namespace FliteSpeechStreaming
+FliteSpeechStreaming::FOnSynthesizedSpeechChunk FliteSpeechStreaming::OnSynthesizedSpeechChunk;
+namespace PrivateFliteSpeechStreaming
 {
 	// written in game thread and read in audio render thread
-	std::atomic<bool> bShouldContinueStreaming(false);
+	static std::atomic<bool> bShouldContinueStreaming(false);
 	/** Scratch buffer used to convert Flite PCM wave data to UE audio engine format of floats */
-	TArray<float> PCMToFloatBuffer;
+	static TArray<float> PCMToFloatBuffer;
 	
 #if SINE_DEBUG
 	Audio::FSineOsc SineOsc;
@@ -80,11 +81,11 @@ namespace FliteSpeechStreaming
 			ReturnCode = CST_AUDIO_STREAM_STOP;
 		}
 		// The delegate should ALWAYS be bound at this point. It's an error to stream without binding this delegate
-		ensure(OnSynthesizedSpeechChunk.IsBound());
-		OnSynthesizedSpeechChunk.ExecuteIfBound(MoveTemp(SpeechData));
+		ensure(FliteSpeechStreaming::OnSynthesizedSpeechChunk.IsBound());
+		FliteSpeechStreaming::OnSynthesizedSpeechChunk.ExecuteIfBound(MoveTemp(SpeechData));
 		return ReturnCode;
 	}
-} // namespace FliteSpeechStreaming
+} // namespace PrivateFliteSpeechStreaming
 
 /** Task to synthesize speech data from a given text in a background thread */
 class FAsyncFliteSpeechSynthesisWorker : public FNonAbandonableTask
@@ -130,13 +131,13 @@ FFliteAdapter::FFliteAdapter()
 	// THis is the ideal size for now as any bigger will result in pops and clicks during playback
 	// as the submix listener can't catch up 
 	FliteAudioStreamingInfo->min_buffsize = 1024;
-	FliteAudioStreamingInfo->asc = FliteSpeechStreaming::flite_audio_stream_chunk;
+	FliteAudioStreamingInfo->asc = PrivateFliteSpeechStreaming::flite_audio_stream_chunk;
 	// Always register voices after creating the streaming info
 	RegisterVoices();
 	
 #if SINE_DEBUG
 	// @TODOAccessibility: Just hardcoding the sample rate for now. Should match that of the voice 
-	FliteSpeechStreaming::SineOsc.Init(16000, 440, 0.5f);
+	PrivateFliteSpeechStreaming::SineOsc.Init(16000, 440, 0.5f);
 #endif
 }
 
@@ -182,14 +183,14 @@ flite_text_to_wave(CharTextPtr, FliteCurrentVoice);
 void FFliteAdapter::StartSynthesizeSpeechData_GameThread(const FString& InText)
 {
 	UE_LOG(LogTextToSpeech, VeryVerbose, TEXT("Toggling streaming flag on."));
-	FliteSpeechStreaming::bShouldContinueStreaming = true;
+	PrivateFliteSpeechStreaming::bShouldContinueStreaming = true;
 	(new FAsyncFliteSpeechSynthesis(*this, InText))->StartBackgroundTask();
 }
 
 void FFliteAdapter::StopSynthesizeSpeechData_GameThread()
 {
 	UE_LOG(LogTextToSpeech, VeryVerbose, TEXT("Toggling speech streaming flag off."));
-	FliteSpeechStreaming::bShouldContinueStreaming = false;
+	PrivateFliteSpeechStreaming::bShouldContinueStreaming = false;
 }
 
 // Flite uses a multiplier to stretch the duratin of the speech 
diff --git a/Engine/Plugins/Experimental/TextToSpeech/Source/TextToSpeech/Private/Flite/FliteAdapter.h b/Engine/Plugins/Experimental/TextToSpeech/Source/TextToSpeech/Private/Flite/FliteAdapter.h
--- a/Engine/Plugins/Experimental/TextToSpeech/Source/TextToSpeech/Private/Flite/FliteAdapter.h
+++ b/Engine/Plugins/Experimental/TextToSpeech/Source/TextToSpeech/Private/Flite/FliteAdapter.h
@@ -5,12 +5,12 @@
 #include "CoreMinimal.h"
 #include "Flite/FliteSynthesizedSpeechData.h"
 
-namespace FliteSpeechStreaming
+struct FliteSpeechStreaming
 {
 	/** Delegate that fires when flite streams in a chunk of audio */
 	DECLARE_DELEGATE_OneParam(FOnSynthesizedSpeechChunk, FFliteSynthesizedSpeechData);
-	FOnSynthesizedSpeechChunk OnSynthesizedSpeechChunk;
-} // namespace Flite
+	static FOnSynthesizedSpeechChunk OnSynthesizedSpeechChunk;
+};
 
 struct cst_voice_struct;
 struct cst_audio_streaming_info_struct;
