Fixed static analysis issues to do with loop termination.

diff --git a/Engine/Plugins/Experimental/TextToSpeech/Source/TextToSpeech/Private/Flite/FliteTextToSpeechSubmixListener.cpp b/Engine/Plugins/Experimental/TextToSpeech/Source/TextToSpeech/Private/Flite/FliteTextToSpeechSubmixListener.cpp
--- a/Engine/Plugins/Experimental/TextToSpeech/Source/TextToSpeech/Private/Flite/FliteTextToSpeechSubmixListener.cpp
+++ b/Engine/Plugins/Experimental/TextToSpeech/Source/TextToSpeech/Private/Flite/FliteTextToSpeechSubmixListener.cpp
@@ -293,18 +293,4 @@ bool FFliteTextToSpeechSubmixListener::FSynthesizedSpeechBuffer::IsFull() const
 	return SynthesizedSpeechChunks.GetNextIndex(WriteIndex) == ReadIndex;
 }
 
-#if !UE_BUILD_SHIPPING
-void FFliteTextToSpeechSubmixListener::FSynthesizedSpeechBuffer::DumpStats() const
-{
-	uint32 MemoryUsage = sizeof(TCircularBuffer<FFliteSynthesizedSpeechData>);
-	// for ReadIndex, WriteIndex and ChunkIndex
-	MemoryUsage += 3* sizeof(int32);
-	// Go through entire buffer and caluclate memory 
-	for (int32 CurrentIndex = ReadIndex; CurrentIndex != ReadIndex; CurrentIndex = SynthesizedSpeechChunks.GetNextIndex(CurrentIndex))
-	{
-		MemoryUsage += sizeof(FFliteSynthesizedSpeechData) + SynthesizedSpeechChunks[CurrentIndex].SpeechBuffer.GetAllocatedSize();
-	}
-	UE_LOG(LogTextToSpeech, Log, TEXT("Memory usage of Flite speech buffer: %u kb."), MemoryUsage / 1000);
-}
-#endif
 #endif
diff --git a/Engine/Plugins/Experimental/TextToSpeech/Source/TextToSpeech/Private/Flite/FliteTextToSpeechSubmixListener.h b/Engine/Plugins/Experimental/TextToSpeech/Source/TextToSpeech/Private/Flite/FliteTextToSpeechSubmixListener.h
--- a/Engine/Plugins/Experimental/TextToSpeech/Source/TextToSpeech/Private/Flite/FliteTextToSpeechSubmixListener.h
+++ b/Engine/Plugins/Experimental/TextToSpeech/Source/TextToSpeech/Private/Flite/FliteTextToSpeechSubmixListener.h
@@ -61,9 +61,6 @@ private:
 		bool IsEmpty() const;
 		/** Should only be called from audio render thread */
 		const FFliteSynthesizedSpeechData& GetCurrentChunk_AudioRenderThread() const;
-#if !UE_BUILD_SHIPPING
-		void DumpStats() const;
-#endif
 	private:
 		bool IsFull() const;
 		TCircularBuffer<FFliteSynthesizedSpeechData> SynthesizedSpeechChunks;
