Added include to fix non-unity builds and whitespace issues.

diff --git a/Engine/Plugins/Experimental/TextToSpeech/Source/TextToSpeech/Private/Flite/FliteAdapter.cpp b/Engine/Plugins/Experimental/TextToSpeech/Source/TextToSpeech/Private/Flite/FliteAdapter.cpp
--- a/Engine/Plugins/Experimental/TextToSpeech/Source/TextToSpeech/Private/Flite/FliteAdapter.cpp
+++ b/Engine/Plugins/Experimental/TextToSpeech/Source/TextToSpeech/Private/Flite/FliteAdapter.cpp
@@ -4,6 +4,7 @@
 #include "Flite/FliteAdapter.h"
 #include "TextToSpeechLog.h"
 #include <atomic>
+#include "Async/AsyncWork.h"
 // Change to 1 to have sine tone data be written for debug
 #define SINE_DEBUG 0
 #if SINE_DEBUG
@@ -55,20 +56,20 @@ namespace PrivateFliteSpeechStreaming
 			UE_LOG(LogTextToSpeech, VeryVerbose, TEXT("Stopping speech streaming."));
 			return CST_AUDIO_STREAM_STOP;
 		}
-			PCMToFloatBuffer.Reset();
+		PCMToFloatBuffer.Reset();
 		PCMToFloatBuffer.AddZeroed(NumSamples);
 		// Note it's possible for a stream of size 0 to come in with it being the last chunk
-			for (int32 CurrentIndex = 0; CurrentIndex < NumSamples; ++CurrentIndex)
-			{
-				// Flite wave samples are stored as shorts (PCM) we need to convert to float
-				// The magic number is the conversion from PCM to float
-				PCMToFloatBuffer[CurrentIndex] = static_cast<float>(Wave->samples[ChunkStartSampleIndex + CurrentIndex]) / 32768.0f;
+		for (int32 CurrentIndex = 0; CurrentIndex < NumSamples; ++CurrentIndex)
+		{
+			// Flite wave samples are stored as shorts (PCM) we need to convert to float
+			// The magic number is the conversion from PCM to float
+			PCMToFloatBuffer[CurrentIndex] = static_cast<float>(Wave->samples[ChunkStartSampleIndex + CurrentIndex]) / 32768.0f;
 #if SINE_DEBUG
-				// Overwriting speech data with sine wave data
+			// Overwriting speech data with sine wave data
 			float Sample = SineOsc.ProcessAudio() * 0.5f;
-				PCMToFloatBuffer[CurrentIndex] = Sample;
+			PCMToFloatBuffer[CurrentIndex] = Sample;
 #endif
-			}
+		}
 		int32 SampleRate = static_cast<int32>(Wave->sample_rate);
 		int32 NumChannels = static_cast<int32>(Wave->num_channels);
 		FFliteSynthesizedSpeechData SpeechData(MoveTemp(PCMToFloatBuffer), SampleRate, NumChannels);;
@@ -177,7 +178,7 @@ void FFliteAdapter::SynthesizeSpeechData_AnyThread(const FString& InText)
 	check(FliteCurrentVoice);
 	// This starts synthesizing the text
 	// This in turns calls the callback bound in the audio streaming info
-flite_text_to_wave(CharTextPtr, FliteCurrentVoice);
+	flite_text_to_wave(CharTextPtr, FliteCurrentVoice);
 }
 
 void FFliteAdapter::StartSynthesizeSpeechData_GameThread(const FString& InText)
diff --git a/Engine/Plugins/Experimental/TextToSpeech/Source/TextToSpeech/Private/Flite/FliteSynthesizedSpeechData.cpp b/Engine/Plugins/Experimental/TextToSpeech/Source/TextToSpeech/Private/Flite/FliteSynthesizedSpeechData.cpp
--- a/Engine/Plugins/Experimental/TextToSpeech/Source/TextToSpeech/Private/Flite/FliteSynthesizedSpeechData.cpp
+++ b/Engine/Plugins/Experimental/TextToSpeech/Source/TextToSpeech/Private/Flite/FliteSynthesizedSpeechData.cpp
@@ -1,5 +1,6 @@
 // Copyright Epic Games, Inc. All Rights Reserved.
 
+#if USING_FLITE
 #include "Flite/FliteSynthesizedSpeechData.h"
 
 FFliteSynthesizedSpeechData::FFliteSynthesizedSpeechData()
@@ -33,3 +34,4 @@ void FFliteSynthesizedSpeechData::Reset()
 	NumChannels = INDEX_NONE;
 }
 
+#endif
\ No newline at end of file
diff --git a/Engine/Plugins/Experimental/TextToSpeech/Source/TextToSpeech/Private/Flite/FliteSynthesizedSpeechData.h b/Engine/Plugins/Experimental/TextToSpeech/Source/TextToSpeech/Private/Flite/FliteSynthesizedSpeechData.h
--- a/Engine/Plugins/Experimental/TextToSpeech/Source/TextToSpeech/Private/Flite/FliteSynthesizedSpeechData.h
+++ b/Engine/Plugins/Experimental/TextToSpeech/Source/TextToSpeech/Private/Flite/FliteSynthesizedSpeechData.h
@@ -1,6 +1,7 @@
 // Copyright Epic Games, Inc. All Rights Reserved.
 
 #pragma once
+#if USING_FLITE
 #include "CoreMinimal.h"
 
 struct FFliteSynthesizedSpeechData
@@ -20,4 +21,4 @@ struct FFliteSynthesizedSpeechData
 	int32 NumChannels;
 	bool bIsLastChunk;
 };
-
+#endif
diff --git a/Engine/Plugins/Experimental/TextToSpeech/Source/TextToSpeech/Private/Flite/FliteTextToSpeech.cpp b/Engine/Plugins/Experimental/TextToSpeech/Source/TextToSpeech/Private/Flite/FliteTextToSpeech.cpp
--- a/Engine/Plugins/Experimental/TextToSpeech/Source/TextToSpeech/Private/Flite/FliteTextToSpeech.cpp
+++ b/Engine/Plugins/Experimental/TextToSpeech/Source/TextToSpeech/Private/Flite/FliteTextToSpeech.cpp
@@ -1,5 +1,6 @@
 // Copyright Epic Games, Inc. All Rights Reserved.
 
+#if USING_FLITE
 #include "Flite/FliteTextToSpeech.h"
 #include "Flite/FliteAdapter.h"
 #include "Flite/FliteTextToSpeechSubmixListener.h"
@@ -172,3 +173,4 @@ void FFliteTextToSpeech::OnSynthesizedSpeechChunk_AnyThread(FFliteSynthesizedSpe
 }
 }
 
+#endif
\ No newline at end of file
diff --git a/Engine/Plugins/Experimental/TextToSpeech/Source/TextToSpeech/Private/Flite/FliteTextToSpeechSubmixListener.h b/Engine/Plugins/Experimental/TextToSpeech/Source/TextToSpeech/Private/Flite/FliteTextToSpeechSubmixListener.h
--- a/Engine/Plugins/Experimental/TextToSpeech/Source/TextToSpeech/Private/Flite/FliteTextToSpeechSubmixListener.h
+++ b/Engine/Plugins/Experimental/TextToSpeech/Source/TextToSpeech/Private/Flite/FliteTextToSpeechSubmixListener.h
@@ -3,6 +3,7 @@
 #pragma once
 #if USING_FLITE
 #include "CoreMinimal.h"
+#include "Containers/CircularBuffer.h"
 #include "Flite/FliteSynthesizedSpeechData.h"
 #include "AudioMixerDevice.h"
 #include "AudioResampler.h"
