Introduced text to speech engine subsystem to add blueprint support for TTS plugin.
The engine subsystem contains multiple text to speech channels which users can
add, remove and interact with to request text to speech functionality.

diff --git a/Engine/Plugins/Experimental/TextToSpeech/Source/TextToSpeech/Private/TextToSpeechEngineSubsystem.cpp b/Engine/Plugins/Experimental/TextToSpeech/Source/TextToSpeech/Private/TextToSpeechEngineSubsystem.cpp
new file mode 100644
--- /dev/null
+++ b/Engine/Plugins/Experimental/TextToSpeech/Source/TextToSpeech/Private/TextToSpeechEngineSubsystem.cpp
@@ -0,0 +1,240 @@
+// Copyright Epic Games, Inc. All Rights Reserved.
+
+#include "TextToSpeechEngineSubsystem.h"
+#include "TextToSpeechModule.h"
+#include "GenericPlatform/ITextToSpeechFactory.h"
+#include "GenericPlatform/TextToSpeechBase.h"
+#include "TextToSpeechLog.h"
+
+UTextToSpeechEngineSubsystem::UTextToSpeechEngineSubsystem()
+{
+
+}
+
+UTextToSpeechEngineSubsystem::~UTextToSpeechEngineSubsystem()
+{
+	RemoveAllChannels();
+}
+
+void UTextToSpeechEngineSubsystem::Initialize(FSubsystemCollectionBase & Collection)
+{
+	Super::Initialize(Collection);
+	// @TODOAccessibility: Parse a config file or something to populate the channel map with pre-defined channels and settings etc 
+}
+
+void UTextToSpeechEngineSubsystem::Deinitialize()
+{
+	RemoveAllChannels();
+	Super::Deinitialize();
+}
+
+void UTextToSpeechEngineSubsystem::SpeakOnChannel(FName InChannelId, UPARAM(ref, DisplayName = "String To Speak") const FString& InStringToSpeak)
+{
+	TSharedRef<FTextToSpeechBase>* TTS = ChannelIdToTextToSpeechMap.Find(InChannelId);
+	if (TTS)
+	{
+		(*TTS)->Speak(InStringToSpeak);
+	}
+}
+
+void UTextToSpeechEngineSubsystem::StopSpeakingOnChannel(UPARAM(DisplayName = "Channel Id") FName InChannelId)
+{
+	TSharedRef<FTextToSpeechBase>* TTS = ChannelIdToTextToSpeechMap.Find(InChannelId);
+	if (TTS)
+	{
+		return (*TTS)->StopSpeaking();
+	}
+	// @TODOAccessibility: Warning about invlaid channel id etc 
+}
+
+void UTextToSpeechEngineSubsystem::StopSpeakingOnAllChannels()
+{
+	for (TPair<FName, TSharedRef<FTextToSpeechBase>> Channel : ChannelIdToTextToSpeechMap)
+	{
+		Channel.Value->StopSpeaking();
+}
+}
+
+bool UTextToSpeechEngineSubsystem::IsSpeakingOnChannel(UPARAM(DisplayName = "Channel Id") FName InChannelId) const
+{
+	const TSharedRef<FTextToSpeechBase>* TTS = ChannelIdToTextToSpeechMap.Find(InChannelId);
+	if (TTS)
+	{
+		return (*TTS)->IsSpeaking();
+	}
+	return false;
+}
+
+float UTextToSpeechEngineSubsystem::GetVolumeOnChannel(UPARAM(DisplayName = "Channel Id") FName InChannelId)
+{
+	const TSharedRef<FTextToSpeechBase>* TTS = ChannelIdToTextToSpeechMap.Find(InChannelId);
+	if (TTS)
+	{
+		return (*TTS)->GetVolume();
+	}
+	return 0.0f;
+}
+
+void UTextToSpeechEngineSubsystem::SetVolumeOnChannel(UPARAM(DisplayName = "Channel Id") FName InChannelId, UPARAM(DisplayName = "Volume") float InVolume)
+{
+	TSharedRef<FTextToSpeechBase>* TTS = ChannelIdToTextToSpeechMap.Find(InChannelId);
+	if (TTS)
+	{
+		(*TTS)->SetVolume(InVolume);
+	}
+}
+
+float UTextToSpeechEngineSubsystem::GetRateOnChannel(UPARAM(DisplayName = "Channel Id") FName InChannelId) const
+{
+	const TSharedRef<FTextToSpeechBase>* TTS = ChannelIdToTextToSpeechMap.Find(InChannelId);
+	if (TTS)
+	{
+		return (*TTS)->GetRate();
+	}
+	return 0.0f;
+}
+
+void UTextToSpeechEngineSubsystem::SetRateOnChannel(UPARAM(DisplayName = "Channel Id") FName InChannelId, UPARAM(DisplayName = "Rate") float InRate)
+{
+	const TSharedRef<FTextToSpeechBase>* TTS = ChannelIdToTextToSpeechMap.Find(InChannelId);
+	if (TTS)
+	{
+		(*TTS)->SetRate(InRate);
+	}
+}
+
+void UTextToSpeechEngineSubsystem::MuteChannel(UPARAM(DisplayName = "Channel Id") FName InChannelId)
+{
+	TSharedRef<FTextToSpeechBase>* TTS = ChannelIdToTextToSpeechMap.Find(InChannelId);
+	if (TTS)
+	{
+		(*TTS)->Mute();
+	}
+}
+
+void UTextToSpeechEngineSubsystem::UnmuteChannel(UPARAM(DisplayName = "Channel Id") FName InChannelId)
+{
+	TSharedRef<FTextToSpeechBase>* TTS = ChannelIdToTextToSpeechMap.Find(InChannelId);
+	if (TTS)
+	{
+		(*TTS)->Unmute();
+	}
+}
+
+bool UTextToSpeechEngineSubsystem::IsChannelMuted(UPARAM(DisplayName = "Channel Id") FName InChannelId) const
+{
+	const TSharedRef<FTextToSpeechBase>* TTS = ChannelIdToTextToSpeechMap.Find(InChannelId);
+	if (TTS)
+	{
+		return (*TTS)->IsMuted();
+	}
+	return false;
+}
+
+void UTextToSpeechEngineSubsystem::ActivateChannel(UPARAM(DisplayName = "Channel Id") FName InChannelId)
+{
+	TSharedRef<FTextToSpeechBase>* TTS = ChannelIdToTextToSpeechMap.Find(InChannelId);
+	if (TTS)
+	{
+		(*TTS)->Activate();
+	}
+}
+
+void UTextToSpeechEngineSubsystem::ActivateAllChannels()
+{
+	for (TPair < FName, TSharedRef<FTextToSpeechBase>>& Channel : ChannelIdToTextToSpeechMap)
+	{
+		Channel.Value->Activate();
+	}
+}
+
+void UTextToSpeechEngineSubsystem::DeactivateChannel(UPARAM(DisplayName = "Channel Id") FName InChannelId)
+{
+	TSharedRef<FTextToSpeechBase>* TTS = ChannelIdToTextToSpeechMap.Find(InChannelId);
+	if (TTS)
+	{
+		(*TTS)->Deactivate();
+	}
+}
+
+void UTextToSpeechEngineSubsystem::DeactivateAllChannels()
+{
+	for (TPair < FName, TSharedRef<FTextToSpeechBase>>& Channel : ChannelIdToTextToSpeechMap)
+	{
+		Channel.Value->Deactivate();
+	}
+}
+
+bool UTextToSpeechEngineSubsystem::IsChannelActive(UPARAM(DisplayName = "Channel Id") FName InChannelId) const
+{
+	const TSharedRef<FTextToSpeechBase>* TTS = ChannelIdToTextToSpeechMap.Find(InChannelId);
+	if (TTS)
+	{
+		(*TTS)->Deactivate();
+	}
+	return false;
+}
+
+void UTextToSpeechEngineSubsystem::AddDefaultChannel(UPARAM(DisplayName = "New Channel Id") FName InNewChannelId)
+{
+	if (!DoesChannelExist(InNewChannelId))
+	{
+		ITextToSpeechModule& Module = ITextToSpeechModule::Get();
+		// platform factory always guaranteed to be valid 
+		ChannelIdToTextToSpeechMap.Add(InNewChannelId, Module.GetPlatformFactory()->Create());
+	}
+	else
+	{
+		UE_LOG(LogTextToSpeech, Warning, TEXT("The requested channel Id: %s already exists."), *InNewChannelId.ToString());
+	}
+}
+
+void UTextToSpeechEngineSubsystem::AddCustomChannel(UPARAM(DisplayName = "New Channel Id") FName InNewChannelId)
+{
+	if (!DoesChannelExist(InNewChannelId))
+	{
+		ITextToSpeechModule& Module = ITextToSpeechModule::Get();
+		if (TSharedPtr<ITextToSpeechFactory> CustomFactory = Module.GetCustomFactory())
+		{
+			ChannelIdToTextToSpeechMap.Add(InNewChannelId, CustomFactory->Create());
+		}
+		else
+		{
+			UE_LOG(LogTextToSpeech, Warning, TEXT("Requested using custom text to speech, but no custom text to speech factory was set."));
+		}
+	}
+	else
+	{
+		UE_LOG(LogTextToSpeech, Warning, TEXT("The requested channel Id: %s already exists."), *InNewChannelId.ToString());
+	}
+}
+
+void UTextToSpeechEngineSubsystem::RemoveChannel(UPARAM(DisplayName = "Channel Id") FName InChannelId)
+{
+	TSharedRef<FTextToSpeechBase>* TTS = ChannelIdToTextToSpeechMap.Find(InChannelId);
+	if (TTS)
+	{
+		(*TTS)->Deactivate();
+		ChannelIdToTextToSpeechMap.Remove(InChannelId);
+	}
+	else
+	{
+		UE_LOG(LogTextToSpeech, Warning, TEXT("Channel Id: %s does not exist. Cannot be removed."), *InChannelId.ToString());
+	}
+}
+
+void UTextToSpeechEngineSubsystem::RemoveAllChannels()
+{
+	DeactivateAllChannels();
+	ChannelIdToTextToSpeechMap.Empty();
+}
+
+bool UTextToSpeechEngineSubsystem::DoesChannelExist(UPARAM(DisplayName = "Channel Id") FName InChannelId) const
+{
+	return ChannelIdToTextToSpeechMap.Contains(InChannelId);
+}
+
+int32 UTextToSpeechEngineSubsystem::GetNumChannels() const
+{
+	return ChannelIdToTextToSpeechMap.Num();
+}
\ No newline at end of file
diff --git a/Engine/Plugins/Experimental/TextToSpeech/Source/TextToSpeech/Public/TextToSpeech.h b/Engine/Plugins/Experimental/TextToSpeech/Source/TextToSpeech/Public/TextToSpeech.h
--- a/Engine/Plugins/Experimental/TextToSpeech/Source/TextToSpeech/Public/TextToSpeech.h
+++ b/Engine/Plugins/Experimental/TextToSpeech/Source/TextToSpeech/Public/TextToSpeech.h
@@ -8,4 +8,5 @@
  */
 #include "GenericPlatform/TextToSpeechBase.h"
 #include "TextToSpeechModule.h"
-#include "GenericPlatform/ITextToSpeechFactory.h"
\ No newline at end of file
+#include "GenericPlatform/ITextToSpeechFactory.h"
+#include "TextToSpeechEngineSubsystem.h"
\ No newline at end of file
diff --git a/Engine/Plugins/Experimental/TextToSpeech/Source/TextToSpeech/Public/TextToSpeechengineSubsystem.h b/Engine/Plugins/Experimental/TextToSpeech/Source/TextToSpeech/Public/TextToSpeechengineSubsystem.h
new file mode 100644
--- /dev/null
+++ b/Engine/Plugins/Experimental/TextToSpeech/Source/TextToSpeech/Public/TextToSpeechengineSubsystem.h
@@ -0,0 +1,180 @@
+// Copyright Epic Games, Inc. All Rights Reserved.
+
+#pragma once
+#include "CoreMinimal.h"
+#include "Subsystems/EngineSubsystem.h"
+#include "TextToSpeechEngineSubsystem.generated.h"
+
+
+class FTextToSpeechBase;
+
+/**
+* Subsystem for interacting with the text to speech system via blueprints.
+* The subsystem consists of a number of text to speech channels.
+* By specifying a channel Id for the functions, you can interact with the channel and request text to be spoken, muting a channel and activating/deactivating a channel.
+* When a channel is added, the channel must then be activated before any TTS operations can be performed.
+*/
+UCLASS()
+class TEXTTOSPEECH_API UTextToSpeechEngineSubsystem: public UEngineSubsystem
+{
+	GENERATED_BODY()
+public:
+	UTextToSpeechEngineSubsystem();
+	virtual ~UTextToSpeechEngineSubsystem();
+	/**
+	* Immediately speaks the requested string asynchronously on the requested TTS channel.
+	* If another string is being spoken by the TTS, the current string will immediately be stopped
+	* before the requested string is spoken.
+	* If the provided channel Id does not exist, nothing will be spoken.
+	* Before executing this function, a TTS channel will need to be created and then activated.
+	* To create a platform default TTS channel, use AddDefaultChannel.
+	* To create a custom TTS channel with a user-defined TTS, use AddCustomChannel.
+	* Note: Passing in long sentences or paragraphs are not recommended. Split the paragraph into shorter sentences for best results.
+	* @param InChannelId The Id of the channel to speak on.
+	* @param InStringToSpeak The string to speak on the requested channel.
+	* @see AddDefaultChannel, AddCustomChannel, ActivateChannel, ActivateAllChannels
+	*/
+	UFUNCTION(BlueprintCallable, Category=TextToSpeech)
+	void SpeakOnChannel(UPARAM(DisplayName = "Channel Id") FName InChannelId, UPARAM(ref, DisplayName="String To Speak") const FString& InStringToSpeak);
+
+	/** 
+	* Immediately stops any text from being spoken on a TTS channel.
+	* @param InChannelId The Id of the channel speech should be stopped on. 
+	*/
+	UFUNCTION(BlueprintCallable, Category = TextToSpeech)
+	void StopSpeakingOnChannel(UPARAM(DisplayName = "Channel Id") FName InChannelId);
+
+	/** Immediately stops text from being spoken on all TTS channels. */
+	UFUNCTION(BlueprintCallable, Category=TextToSpeech)
+	void StopSpeakingOnAllChannels();
+
+	/** 
+	* Returns true if text is being spoken on a TTS channel. Else false.
+	* If the provided channel Id doesn't exist, function will return false.
+	* @param InChannelId The id of the channel to check if text is being spoken. 
+	*/
+	UFUNCTION(BlueprintCallable, Category=TextToSpeech)
+	bool IsSpeakingOnChannel(UPARAM(DisplayName = "Channel Id") FName InChannelId) const;
+
+	/** 
+	* Returns the current volume being spoken on a TTS channel.  Value is between 0.0f and 1.0f.
+	* If the provided channel Id doesn't exist, 0.0f will be returned. 
+	* @param InChannelId The Id for the channel to retrieve the volume from.
+	*/
+	UFUNCTION(BlueprintCallable, Category = TextToSpeech)
+	float GetVolumeOnChannel(UPARAM(DisplayName = "Channel Id") FName InChannelId);
+	/** 
+	* Sets the volume for strings to be spoken on a TTS channel. 
+	* If the provided channel Id does not exist, nothing will happen.
+	* @param InChannelId The Id for the channel to set for.
+	* @param InVolume The volume to be set on the channel. The value will be clamped between 0.0f and 1.0f. 
+	*/
+	UFUNCTION(BlueprintCallable, Category = TextToSpeech)
+	void SetVolumeOnChannel(UPARAM(DisplayName = "Channel Id") FName InChannelId, UPARAM(DisplayName = "Volume") float InVolume);
+	/** 
+	* Returns the current speech rate strings are spoken on a TTS channel. Value is between 0 and 1.0f.
+	* If the provided channel Id does not exist, 0.0f will be returned. 
+	* @param InChannelId The Id for the channel to get the speech rate from. 
+	*/
+	UFUNCTION(BlueprintCallable, Category = TextToSpeech)
+	float GetRateOnChannel(UPARAM(DisplayName = "Channel Id") FName InChannelId) const;
+	/** 
+	* Sets the current speech rate strings should be spoken on a TTS channel.
+	* If the provided channel does not exist, nothing will happen. 
+	* @param InChannelId The Id for the channel to set the speech rate on.
+	* @param InRate The speech rate to set for the channel. Value will be clamped between 0.0f and 1.0f.
+	*/
+	UFUNCTION(BlueprintCallable, Category = TextToSpeech)
+	void SetRateOnChannel(UPARAM(DisplayName = "Channel Id") FName InChannelId, UPARAM(DisplayName = "Rate") float InRate);
+	/** 
+	* Mutes a TTS channel so no synthesized speech is audible on that channel.
+	* @param InChannelId The Id for the channel to mute.
+	*/
+	UFUNCTION(BlueprintCallable, Category = TextToSpeech)
+	void MuteChannel(UPARAM(DisplayName = "Channel Id") FName InChannelId);
+	/** 
+	* Unmutes a TTS channel so requests to speak strings are audible again on the channel.
+	@param InChannelId The Id for the channel to unmute.
+	*/
+	UFUNCTION(BlueprintCallable, Category = TextToSpeech)
+	void UnmuteChannel(UPARAM(DisplayName = "Channel Id") FName InChannelId);
+	/** 
+	* Returns true if the TTS channel is muted. Else false.
+	* If the provided channel Id does not exist, the function will return false.
+	* @param InChannelId The Id for the channel to check if it is muted.
+	*/
+	UFUNCTION(BlueprintCallable, Category = TextToSpeech)
+	bool IsChannelMuted(UPARAM(DisplayName = "Channel Id") FName InChannelId) const;
+	/**
+	* Activates a TTS channel making it ready to synthesize text and read out strings.
+	* If the provided channel Id does not exist, nothing will happen. 
+	* @param InChannelId The Id of the channel to activate. 
+	*/
+	UFUNCTION(BlueprintCallable, Category = TextToSpeech)
+	void ActivateChannel(UPARAM(DisplayName="Channel Id") FName InChannelId);
+	/** Activates all TTS channels making themready to synthesize text and read out strings. */
+	UFUNCTION(BlueprintCallable, Category = TextToSpeech)
+	void ActivateAllChannels();
+	/**
+	* Deactivates a TTS channel making all requests for speaking text for that channel do nothing.
+	* If the provided channel Id does not exist, nothing will happen. 
+	* @param InChannelId The Id for the channel to deactivate.
+	*/
+	UFUNCTION(BlueprintCallable, Category=TextToSpeech)
+	void DeactivateChannel(UPARAM(DisplayName = "Channel Id") FName InChannelId);
+	/** Deactivates all TTS channels making all requests for speaking text do nothing. */
+	UFUNCTION(BlueprintCallable, Category = TextToSpeech)
+	void DeactivateAllChannels();
+	/** 
+	* Returns true if the TTS channel is active. Else false.
+	* If the provided channel Id does not exist, returns false.
+	* @param InChannelId The Id for the channel to check if it is active. 
+	*/
+	UFUNCTION(BlueprintCallable, Category = TextToSpeech)
+	bool IsChannelActive(UPARAM(DisplayName = "Channel Id") FName InChannelId) const;
+	/**
+	* Creates a new TTS channel for text to speech requests to be made to a platform text to speech engine. 
+	* This will not create the channel if the provided channel id is not unique.
+	* Channels added using this function start off deactivated. They must be activated to use them for text to speech functionality. 
+	* For out-of-the-box TTS support, this is most likely the channel creation method you want.
+	* @param InChannelId The Id of the channel you want to add
+	* @see ActivateChannel, ActivateAllChannels
+	*/
+	UFUNCTION(BlueprintCallable, Category=TextToSpeech)
+	void AddDefaultChannel(UPARAM(DisplayName = "New Channel Id") FName InNewChannelId);
+	/** 
+	* Creates a new TTS channel where text to speech requests are fulfilled by a user defined TTS.
+	* If you have not specified a custom TTS to be used, use AddDefaultChannel instead.
+	* This will not add a channel if the channel Id is not unique or if the user has not specified a custom TTS to be used in ITextToSpeechModule.
+	* Channels added using this function start off deactivated. They must be activated to use them for text to speech functionality. 
+	* @see ITextToSpeechModule, AddDefaultChannel, ActivateChannel, ActivateAllChannels
+	*/
+	UFUNCTION(BlueprintCallable, Category = TextToSpeech)
+	void AddCustomChannel(UPARAM(DisplayName = "New Channel Id") FName InNewChannelId);
+	/** 
+	* Removes a TTS channel, preventing all further requests for text to speech functionality from the channel. 
+	* If the provided channel Id does not exist, nothing will happen.
+	* @param InChannelId The Id for the channel you want removed. 
+	*/
+	UFUNCTION(BlueprintCallable, Category=TextToSpeech)
+	void RemoveChannel(UPARAM(DisplayName = "Channel Id") FName InChannelId);
+	/** Removes all TTS channels.*/ 
+	UFUNCTION(BlueprintCallable, Category = TextToSpeech)
+		void RemoveAllChannels();
+	/**
+	* Returns true if a TTS channel associated with a channel Id exists. Else false.
+	* @param InChannelId The channel Id to test if it is associated with a channel.
+	*/
+	UFUNCTION(BlueprintCallable, Category=TextToSpeech)
+	bool DoesChannelExist(UPARAM(DisplayName = "Channel Id") FName InChannelId) const;
+	/** Returns the number of TTS channels currently created. */
+	UFUNCTION(BlueprintCallable, Category=TextToSpeech)
+	int32 GetNumChannels() const;
+// UEngineSubsystem
+	virtual void Initialize(FSubsystemCollectionBase& Collection) override;
+	virtual void Deinitialize() override;
+	// ~
+private:
+	/** The map of channel Ids to native text to speech objects. */
+	TMap<FName, TSharedRef<FTextToSpeechBase>> ChannelIdToTextToSpeechMap;
+};
diff --git a/Engine/Plugins/Experimental/TextToSpeech/Source/TextToSpeech/TextToSpeech.Build.cs b/Engine/Plugins/Experimental/TextToSpeech/Source/TextToSpeech/TextToSpeech.Build.cs
--- a/Engine/Plugins/Experimental/TextToSpeech/Source/TextToSpeech/TextToSpeech.Build.cs
+++ b/Engine/Plugins/Experimental/TextToSpeech/Source/TextToSpeech/TextToSpeech.Build.cs
@@ -10,7 +10,8 @@ public class TextToSpeech : ModuleRules
 		PublicDependencyModuleNames.AddRange(
 			new string[] 
 			{
-				"Core"
+				"Core",
+				"CoreUObject"
 			}
 		);
 	
