Fix for incremental build issue on Mac for text to speech engine subsystem.

diff --git a/Engine/Plugins/Experimental/TextToSpeech/Source/TextToSpeech/TextToSpeech.Build.cs b/Engine/Plugins/Experimental/TextToSpeech/Source/TextToSpeech/TextToSpeech.Build.cs
--- a/Engine/Plugins/Experimental/TextToSpeech/Source/TextToSpeech/TextToSpeech.Build.cs
+++ b/Engine/Plugins/Experimental/TextToSpeech/Source/TextToSpeech/TextToSpeech.Build.cs
@@ -23,13 +23,19 @@ public class TextToSpeech : ModuleRules
 				"ApplicationCore"
 			}
 		);
+		
+		PrivateDependencyModuleNames.AddRange(
+			new string[]
+			{
+				"Engine"
+			}
+		);
 		if (Target.Platform == UnrealTargetPlatform.Win64)
 		{
 			PrivateDependencyModuleNames.AddRange(
 				new string[] 
 				{
 					"AudioMixer",
-					"Engine",
 					"SignalProcessing",
 					"AudioPlatformConfiguration"
 				}
