Fixed bug with Flite text to speech where muting or unmuting a deactivated deactivated TTS object will cause a crash.

diff --git a/Engine/Plugins/Experimental/TextToSpeech/Source/TextToSpeech/Private/Flite/FliteTextToSpeech.cpp b/Engine/Plugins/Experimental/TextToSpeech/Source/TextToSpeech/Private/Flite/FliteTextToSpeech.cpp
--- a/Engine/Plugins/Experimental/TextToSpeech/Source/TextToSpeech/Private/Flite/FliteTextToSpeech.cpp
+++ b/Engine/Plugins/Experimental/TextToSpeech/Source/TextToSpeech/Private/Flite/FliteTextToSpeech.cpp
@@ -122,8 +122,11 @@ void FFliteTextToSpeech::Mute()
 void FFliteTextToSpeech::Unmute()
 {
 	check(IsInGameThread());
-	SetMuted(false);
-	TTSSubmixListener->Unmute();
+	if (IsActive())
+	{
+		SetMuted(false);
+		TTSSubmixListener->Unmute();
+	}
 }
 
 void FFliteTextToSpeech::OnActivated()
