Fixed bug where text to speech engine subsystem returns a wrong value when querying if a channel is active.

diff --git a/Engine/Plugins/Experimental/TextToSpeech/Source/TextToSpeech/Private/TextToSpeechEngineSubsystem.cpp b/Engine/Plugins/Experimental/TextToSpeech/Source/TextToSpeech/Private/TextToSpeechEngineSubsystem.cpp
--- a/Engine/Plugins/Experimental/TextToSpeech/Source/TextToSpeech/Private/TextToSpeechEngineSubsystem.cpp
+++ b/Engine/Plugins/Experimental/TextToSpeech/Source/TextToSpeech/Private/TextToSpeechEngineSubsystem.cpp
@@ -170,7 +170,7 @@ bool UTextToSpeechEngineSubsystem::IsChannelActive(UPARAM(DisplayName = "Channel
 	const TSharedRef<FTextToSpeechBase>* TTS = ChannelIdToTextToSpeechMap.Find(InChannelId);
 	if (TTS)
 	{
-		(*TTS)->Deactivate();
+		return (*TTS)->IsActive();
 	}
 	return false;
 }
