Introduced text to speech support for IOS by using platform text to speech.

diff --git a/Engine/Plugins/Experimental/TextToSpeech/Source/TextToSpeech/Private/IOS/IOSTextToSpeech.cpp b/Engine/Plugins/Experimental/TextToSpeech/Source/TextToSpeech/Private/IOS/IOSTextToSpeech.cpp
new file mode 100644
--- /dev/null
+++ b/Engine/Plugins/Experimental/TextToSpeech/Source/TextToSpeech/Private/IOS/IOSTextToSpeech.cpp
@@ -0,0 +1,188 @@
+// Copyright Epic Games, Inc. All Rights Reserved.
+
+#if PLATFORM_IOS
+ #include "IOS/IOSTextToSpeech.h"
+#include "TextToSpeechLog.h"
+#include "IOS/IOSAppDelegate.h"
+#import <UIKit/UIKit.h>
+#import <AVFoundation/AVFoundation.h>
+
+/** Defined in FTextToSpeechBase.cpp */
+extern TMap<TextToSpeechId, TWeakPtr<FTextToSpeechBase>> ActiveTextToSpeechMap;
+
+@interface FSpeechSynthesizerDelegate : FApplePlatformObject<AVSpeechSynthesizerDelegate>
+-(id)initWithOwningTextToSpeechId:(TextToSpeechId)InOwningTextToSpeechId;
+-(void)dealloc;
+-(void)speechSynthesizer:(AVSpeechSynthesizer *)synthesizer didFinishSpeechUtterance:(AVSpeechUtterance *)utterance;
+
+@end
+
+@implementation FSpeechSynthesizerDelegate
+{
+	TextToSpeechId _OwningTextToSpeechId;
+}
+APPLE_PLATFORM_OBJECT_ALLOC_OVERRIDES(FSpeechSynthesizerDelegate)
+-(id)initWithOwningTextToSpeechId:(TextToSpeechId)InOwningTextToSpeechId
+{
+	if (self = [super init])
+	{
+		_OwningTextToSpeechId = InOwningTextToSpeechId;
+	}
+	return self;
+}
+
+-(void)dealloc
+{
+	_OwningTextToSpeechId = FTextToSpeechBase::InvalidTextToSpeechId;
+	[super dealloc];
+}
+
+-(void)speechSynthesizer:(AVSpeechSynthesizer*)synthesizer didFinishSpeechUtterance:(AVSpeechUtterance*)utterance
+{
+	if (_OwningTextToSpeechId != FTextToSpeechBase::InvalidTextToSpeechId)
+	{
+		const TextToSpeechId IdCopy = _OwningTextToSpeechId;
+		const TMap<TextToSpeechId, TWeakPtr<FTextToSpeechBase>>& ActiveTextToSpeechMapRef = ActiveTextToSpeechMap;
+		FFunctionGraphTask::CreateAndDispatchWhenReady([IdCopy, ActiveTextToSpeechMapRef]()
+			{
+				const TWeakPtr<FTextToSpeechBase>* TTSPtr = ActiveTextToSpeechMapRef.Find(IdCopy);
+				if (TTSPtr && TTSPtr->IsValid())
+				{
+					TTSPtr->Pin()->OnTextToSpeechFinishSpeaking_GameThread();
+				}
+			}, TStatId(), NULL, ENamedThreads::GameThread);
+	}
+}
+@end
+
+FIOSTextToSpeech::FIOSTextToSpeech()
+	: bIsSpeaking(false)
+	, Volume(0.0f)
+	, Rate(0.0f)
+	, SpeechSynthesizer(nullptr)
+	, SpeechSynthesizerDelegate(nullptr)
+{
+
+}
+
+FIOSTextToSpeech::~FIOSTextToSpeech()
+{
+
+}
+
+void FIOSTextToSpeech::Speak(const FString& InStringToSpeak)
+{
+	if (IsActive())
+	{
+		if (!InStringToSpeak.IsEmpty())
+		{
+			if (IsSpeaking())
+			{
+				StopSpeaking();
+			}
+			SCOPED_AUTORELEASE_POOL;
+			NSString* Announcement = InStringToSpeak.GetNSString();
+			ensureMsgf(Announcement, TEXT("Failed to convert FString to NSString for TTS."));
+			AVSpeechUtterance* Utterance = [AVSpeechUtterance speechUtteranceWithString:Announcement];
+			checkf(Utterance, TEXT("Failed to create an utterance for text to speech."));
+			// for now we just use the default system language that's being used
+			Utterance.voice = [AVSpeechSynthesisVoice voiceWithLanguage: [AVSpeechSynthesisVoice currentLanguageCode]];
+			// If muted, set to volume to 0
+			Utterance.volume = IsMuted() ? 0.0f : Volume;
+			Utterance.rate = Rate;
+			[SpeechSynthesizer speakUtterance:Utterance];
+			bIsSpeaking = true;
+		}
+	}
+}
+
+bool FIOSTextToSpeech::IsSpeaking() const
+{
+	return IsActive() ? bIsSpeaking : false;
+}
+
+void FIOSTextToSpeech::StopSpeaking()
+{
+	if (IsActive())
+	{
+		if (IsSpeaking())
+		{
+			[SpeechSynthesizer stopSpeakingAtBoundary:AVSpeechBoundaryImmediate];
+		}
+		bIsSpeaking = false;
+	}
+}
+
+float FIOSTextToSpeech::GetVolume() const
+{
+	return Volume;
+}
+
+void FIOSTextToSpeech::SetVolume(float InVolume)
+{
+	Volume = FMath::Clamp(InVolume, 0.0f, 1.0f);
+}
+
+float FIOSTextToSpeech::GetRate() const
+{
+	return Rate;
+}
+
+void FIOSTextToSpeech::SetRate(float InRate)
+{
+	Rate = FMath::Clamp(InRate, 0.0f, 1.0f);
+}
+
+void FIOSTextToSpeech::Mute()
+{
+	if (!IsMuted())
+	{
+		SetMuted(true);
+	}
+}
+
+void FIOSTextToSpeech::Unmute()
+{
+	if (IsMuted())
+	{
+		SetMuted(false);
+	}
+}
+
+void FIOSTextToSpeech::OnActivated()
+{
+	ensureMsgf(!IsActive(), TEXT("Attempting to activate an already activated TTS. FTextToSpeechBase::Activate() should already guard against this."));
+	TextToSpeechId IdCopy = GetId();
+	dispatch_async(dispatch_get_main_queue(), ^
+	{
+		SpeechSynthesizer = [[AVSpeechSynthesizer alloc] init];
+		SpeechSynthesizerDelegate = [[FSpeechSynthesizerDelegate alloc] initWithOwningTextToSpeechId:IdCopy];
+		SpeechSynthesizer.delegate = SpeechSynthesizerDelegate;
+		// To get initial volume and rate, we need to retrieve it from an utterance
+		AVSpeechUtterance* Utterance = [AVSpeechUtterance speechUtteranceWithString:@"Temp"];
+		checkf(Utterance, TEXT("Failed to create an utterance for TTS."));
+		Volume = Utterance.volume;
+		Rate = Utterance.rate;
+		ensure(Volume >= 0 && Volume <= 1.0f);
+		ensure(Rate >= 0 && Rate <= 1.0f);
+		// This allows us to still hear the TTS when the IOS ringer is muted
+		// SetFeature already takes care of keeping track of how many requests are made to activate/deactivate an audio feature
+		[[IOSAppDelegate GetDelegate] SetFeature:EAudioFeature::Playback Active:true];
+	});
+}
+
+void FIOSTextToSpeech::OnDeactivated()
+{
+	ensureMsgf(IsActive(), TEXT("Attempting to deactivate an already deactivated TTS. FTextToSpeechBase::Deactivate() should already guard against this."));
+	// deallocate all AppKit objects in main thread just in case
+	dispatch_async(dispatch_get_main_queue(), ^
+	{
+		checkf(SpeechSynthesizerDelegate, TEXT("Null speech synthesizer in IOS tex to speech. A valid delegate must exist for the lifetime of the TTS."));
+		[SpeechSynthesizerDelegate release];
+		checkf(SpeechSynthesizer, TEXT("Null speech synthesizer in IOS text to speech. There should be a valid speech  synthesizer for the lifetime of the TTS."));
+		[SpeechSynthesizer release];
+		// SetFeature already takes care of keeping track of the number of requests to activate/deactivate a feature
+		[[IOSAppDelegate GetDelegate] SetFeature:EAudioFeature::Playback Active:false];
+	});
+};
+#endif
diff --git a/Engine/Plugins/Experimental/TextToSpeech/Source/TextToSpeech/Private/IOS/IOSTextToSpeech.h b/Engine/Plugins/Experimental/TextToSpeech/Source/TextToSpeech/Private/IOS/IOSTextToSpeech.h
new file mode 100644
--- /dev/null
+++ b/Engine/Plugins/Experimental/TextToSpeech/Source/TextToSpeech/Private/IOS/IOSTextToSpeech.h
@@ -0,0 +1,47 @@
+// Copyright Epic Games, Inc. All Rights Reserved.
+
+#pragma once
+#if PLATFORM_IOS
+#include "CoreMinimal.h"
+#include "GenericPlatform/TextToSpeechBase.h"
+
+@class AVSpeechSynthesizer;
+@class FSpeechSynthesizerDelegate;
+
+class FIOSTextToSpeech : public FTextToSpeechBase
+{
+public:
+	FIOSTextToSpeech();
+	virtual ~FIOSTextToSpeech();
+	
+	// FGenericTextToSpeech
+	virtual void Speak(const FString& InStringToSpeak) override final;
+	virtual bool IsSpeaking() const override final;
+	virtual void StopSpeaking() override final;
+	virtual float GetVolume() const override final;
+	virtual void SetVolume(float InVolume) override final;
+	virtual float GetRate() const override final;
+	virtual void SetRate(float InRate) override final;
+	virtual void Mute() override final;
+	virtual void Unmute() override final;
+protected:
+	virtual void OnActivated() override final;
+	virtual void OnDeactivated() override final;
+	// ~
+private:
+	/** True if the speech synthesizer *is speaking. Else false. */
+	bool bIsSpeaking;
+	/** The current volume all TTS utterances will be played at **/
+	float Volume;
+	/** The speech rate all TTS utterances will be made at. */
+	float Rate;
+	/** The platform speech synthesizer that converts text to speech */
+	AVSpeechSynthesizer* SpeechSynthesizer;
+	/**
+	* The delegate for the speech synthesizer. Callbacks indicating
+	* speech synthesis playback progress, intruuption and completion are handled by this delegate.
+	 */
+	FSpeechSynthesizerDelegate* SpeechSynthesizerDelegate;
+};
+#endif
+
diff --git a/Engine/Plugins/Experimental/TextToSpeech/Source/TextToSpeech/Private/IOS/IOSTextToSpeechFactory.cpp b/Engine/Plugins/Experimental/TextToSpeech/Source/TextToSpeech/Private/IOS/IOSTextToSpeechFactory.cpp
new file mode 100644
--- /dev/null
+++ b/Engine/Plugins/Experimental/TextToSpeech/Source/TextToSpeech/Private/IOS/IOSTextToSpeechFactory.cpp
@@ -0,0 +1,13 @@
+// Copyright Epic Games, Inc. All Rights Reserved.
+
+#if PLATFORM_IOS
+#include "IOS/IOSTextToSpeechFactory.h"
+#include "GenericPlatform/ITextToSpeechFactory.h"
+#include "IOS/IOSTextToSpeech.h"
+
+TSharedRef<FTextToSpeechBase> FIOSTextToSpeechFactory::Create()
+{
+	return MakeShared<FIOSTextToSpeech>();
+}
+
+#endif
diff --git a/Engine/Plugins/Experimental/TextToSpeech/Source/TextToSpeech/Private/IOS/IOSTextToSpeechFactory.h b/Engine/Plugins/Experimental/TextToSpeech/Source/TextToSpeech/Private/IOS/IOSTextToSpeechFactory.h
new file mode 100644
--- /dev/null
+++ b/Engine/Plugins/Experimental/TextToSpeech/Source/TextToSpeech/Private/IOS/IOSTextToSpeechFactory.h
@@ -0,0 +1,20 @@
+// Copyright Epic Games, Inc. All Rights Reserved.
+
+#pragma once
+#if PLATFORM_IOS
+#include "CoreMinimal.h"
+#include "GenericPlatform/ITextToSpeechFactory.h"
+
+class FTextToSpeechBase;
+/** The platform default text to speech factory for IOS. */
+class FIOSTextToSpeechFactory : public ITextToSpeechFactory
+{
+public:
+	FIOSTextToSpeechFactory() = default;
+	virtual ~FIOSTextToSpeechFactory() = default;
+
+	virtual TSharedRef<FTextToSpeechBase> Create() override;
+};
+typedef FIOSTextToSpeechFactory FPlatformTextToSpeechFactory;
+#endif
+
diff --git a/Engine/Plugins/Experimental/TextToSpeech/TextToSpeech.uplugin b/Engine/Plugins/Experimental/TextToSpeech/TextToSpeech.uplugin
--- a/Engine/Plugins/Experimental/TextToSpeech/TextToSpeech.uplugin
+++ b/Engine/Plugins/Experimental/TextToSpeech/TextToSpeech.uplugin
@@ -25,6 +25,7 @@
 				"Win64",
 				"Mac",
 				"Android",
+				"IOS"
 			],
 			"BlacklistTargets" :[
 				"Server"
