Fix to comment for tooltips on Text To Speech SubSystem.

diff --git a/Engine/Plugins/Experimental/TextToSpeech/Source/TextToSpeech/Public/TextToSpeechEngineSubsystem.h b/Engine/Plugins/Experimental/TextToSpeech/Source/TextToSpeech/Public/TextToSpeechEngineSubsystem.h
--- a/Engine/Plugins/Experimental/TextToSpeech/Source/TextToSpeech/Public/TextToSpeechEngineSubsystem.h
+++ b/Engine/Plugins/Experimental/TextToSpeech/Source/TextToSpeech/Public/TextToSpeechEngineSubsystem.h
@@ -22,6 +22,7 @@ class TEXTTOSPEECH_API UTextToSpeechEngineSubsystem: public UEngineSubsystem
 public:
 	UTextToSpeechEngineSubsystem();
 	virtual ~UTextToSpeechEngineSubsystem();
+
 	/**
 	* Immediately vocalizes the requested string asynchronously on the requested text to speech channel, interrupting any string that is already being vocalized on the channel.
 	* If the provided channel Id does not exist, nothing will be vocalized.
@@ -49,8 +50,7 @@ public:
 	void StopSpeakingOnAllChannels();
 
 	/** 
-	* Returns true if a string is being vocalized on a text to speech channel. Otherwise, the function returns false.
-	* If the provided channel Id doesn't exist, function will return false.
+	* Return true when the targeted text to speech channel is vocalising, otherwise false.
 	* @param InChannelId The id of the channel to check if a string is being vocalized.
 	*/
 	UFUNCTION(BlueprintCallable, Category=TextToSpeech)
@@ -63,21 +63,24 @@ public:
 	*/
 	UFUNCTION(BlueprintCallable, Category = TextToSpeech)
 	float GetVolumeOnChannel(UPARAM(DisplayName = "Channel Id") FName InChannelId);
+
 	/** 
-	* Sets the volume for strings to be vocalized on a text to speech channel. 
+	* Sets the volume for strings vocalized on a text to speech channel. 
 	* If the provided channel Id does not exist, nothing will happen.
 	* @param InChannelId The Id for the channel to set for.
 	* @param InVolume The volume to be set on the channel. The value will be clamped between 0.0f and 1.0f. 
 	*/
 	UFUNCTION(BlueprintCallable, Category = TextToSpeech)
 	void SetVolumeOnChannel(UPARAM(DisplayName = "Channel Id") FName InChannelId, UPARAM(DisplayName = "Volume") float InVolume);
+
 	/** 
-	* Returns the current speech rate strings are vocalized on a text to speech channel. Value is between 0 and 1.0f.
+	* Returns the current speech rate strings are vocalized on a text to speech channel. Value is between 0.0 and 1.0.
 	* If the provided channel Id does not exist, 0.0f will be returned. 
 	* @param InChannelId The Id for the channel to get the speech rate from. 
 	*/
 	UFUNCTION(BlueprintCallable, Category = TextToSpeech)
 	float GetRateOnChannel(UPARAM(DisplayName = "Channel Id") FName InChannelId) const;
+
 	/** 
 	* Sets the current speech rate strings should be vocalized on a text to speech channel.
 	* If the provided channel does not exist, nothing will happen. 
@@ -86,6 +89,7 @@ public:
 	*/
 	UFUNCTION(BlueprintCallable, Category = TextToSpeech)
 	void SetRateOnChannel(UPARAM(DisplayName = "Channel Id") FName InChannelId, UPARAM(DisplayName = "Rate") float InRate);
+
 	/** 
 	* Mutes a text to speech channel so no vocalized strings are audible on that channel.
 	* If the requested channel is already muted, nothing will happen.
@@ -93,6 +97,7 @@ public:
 	*/
 	UFUNCTION(BlueprintCallable, Category = TextToSpeech)
 	void MuteChannel(UPARAM(DisplayName = "Channel Id") FName InChannelId);
+
 	/** 
 	* Unmutes a text to speech channel so vocalized strings are audible on the channel.
 	* If the requested channel is already unmuted, nothing will happen.
@@ -100,13 +105,14 @@ public:
 	*/
 	UFUNCTION(BlueprintCallable, Category = TextToSpeech)
 	void UnmuteChannel(UPARAM(DisplayName = "Channel Id") FName InChannelId);
+
 	/** 
-	* Returns true if the text to speech channel is muted. Otherwise, the function returns false.
-	* If the provided channel Id does not exist, the function will return false.
+	* Returns true if the text to speech channel is muted, otherwise false.
 	* @param InChannelId The Id for the channel to check if it is muted.
 	*/
 	UFUNCTION(BlueprintCallable, Category = TextToSpeech)
 	bool IsChannelMuted(UPARAM(DisplayName = "Channel Id") FName InChannelId) const;
+
 	/**
 	* Activates a text to speech channel to accept requests to perform text to speech functionality.
 	* If the provided channel Id does not exist, nothing will happen. 
@@ -114,9 +120,11 @@ public:
 	*/
 	UFUNCTION(BlueprintCallable, Category = TextToSpeech)
 	void ActivateChannel(UPARAM(DisplayName="Channel Id") FName InChannelId);
+
 	/** Activates all text to speech channels to accept requests for text to speech functionality. */
 	UFUNCTION(BlueprintCallable, Category = TextToSpeech)
 	void ActivateAllChannels();
+
 	/**
 	* Deactivates a text to speech channel and stops any vocalized strings on that channel. Future Requests for text to speech functionality will do nothing.
 	* If the provided channel Id does not exist, nothing will happen. 
@@ -124,16 +132,18 @@ public:
 	*/
 	UFUNCTION(BlueprintCallable, Category=TextToSpeech)
 	void DeactivateChannel(UPARAM(DisplayName = "Channel Id") FName InChannelId);
+
 	/** Deactivates all text to speech channels making all requests for text to speech functionality do nothing. */
 	UFUNCTION(BlueprintCallable, Category = TextToSpeech)
 	void DeactivateAllChannels();
+
 	/** 
-	* Returns true if the text to speech channel is active. Otherwise, the function returns false.
-	* If the provided channel Id does not exist, returns false.
+	* Returns true if the text to speech channel is active, otherwise false.
 	* @param InChannelId The Id for the channel to check if it is active. 
 	*/
 	UFUNCTION(BlueprintCallable, Category = TextToSpeech)
 	bool IsChannelActive(UPARAM(DisplayName = "Channel Id") FName InChannelId) const;
+
 	/**
 	* Creates a new  channel for text to speech requests to be made to a platform C++ text to speech class. 
 	* This will not create the channel if the provided channel id is not unique.
@@ -144,6 +154,7 @@ public:
 	*/
 	UFUNCTION(BlueprintCallable, Category=TextToSpeech)
 	void AddDefaultChannel(UPARAM(DisplayName = "New Channel Id") FName InNewChannelId);
+
 	/** 
 	* Creates a new text to speech channel where text to speech requests are fulfilled by a user implemented C++ text to speech class. 
 	* If you have not specified a custom text to speech class to be used, use AddDefaultChannel instead.
@@ -153,6 +164,7 @@ public:
 	*/
 	UFUNCTION(BlueprintCallable, Category = TextToSpeech)
 	void AddCustomChannel(UPARAM(DisplayName = "New Channel Id") FName InNewChannelId);
+
 	/** 
 	* Removes a text to speech channel, preventing all further requests for text to speech functionality from the channel. 
 	* If the provided channel Id does not exist, nothing will happen.
@@ -160,18 +172,22 @@ public:
 	*/
 	UFUNCTION(BlueprintCallable, Category=TextToSpeech)
 	void RemoveChannel(UPARAM(DisplayName = "Channel Id") FName InChannelId);
+
 	/** Removes all text to speech channels, preventin future requests for text to speech functionality on all channels. */ 
 	UFUNCTION(BlueprintCallable, Category = TextToSpeech)
-		void RemoveAllChannels();
+	void RemoveAllChannels();
+
 	/**
 	* Returns true if a text to speech channel associated with a channel Id exists. Otherwise, the function returns false.
 	* @param InChannelId The channel Id to test if it is associated with a channel.
 	*/
 	UFUNCTION(BlueprintCallable, Category=TextToSpeech)
 	bool DoesChannelExist(UPARAM(DisplayName = "Channel Id") FName InChannelId) const;
+
 	/** Returns the number of text to speech channels that have been added. */
 	UFUNCTION(BlueprintCallable, Category=TextToSpeech, meta=(DisplayName="Get Number Of Channels"))
 	int32 GetNumChannels() const;
+
 // UEngineSubsystem
 	virtual void Initialize(FSubsystemCollectionBase& Collection) override;
 	virtual void Deinitialize() override;
