Fix for screen reader crash when navigating out of a combo box.
- Changed all static or class scratch audio buffers to local variables instead to remove the need for synchronization.
- Created a cvar for the initial size of the Flite speech buffer to reduce the liklihood of running out of memory while streaming in audio data from Flite.
#rnx

diff --git a/Engine/Plugins/Experimental/TextToSpeech/Source/TextToSpeech/Private/Flite/FliteAdapter.cpp b/Engine/Plugins/Experimental/TextToSpeech/Source/TextToSpeech/Private/Flite/FliteAdapter.cpp
--- a/Engine/Plugins/Experimental/TextToSpeech/Source/TextToSpeech/Private/Flite/FliteAdapter.cpp
+++ b/Engine/Plugins/Experimental/TextToSpeech/Source/TextToSpeech/Private/Flite/FliteAdapter.cpp
@@ -35,10 +35,7 @@ FliteSpeechStreaming::FOnSynthesizedSpeechChunk FliteSpeechStreaming::OnSynthesi
 namespace PrivateFliteSpeechStreaming
 {
 	// written in game thread and read in audio render thread
-	static std::atomic<bool> bShouldContinueStreaming(false);
-	/** Scratch buffer used to convert Flite PCM wave data to UE audio engine format of floats */
-	static TArray<float> PCMToFloatBuffer;
-	
+	static std::atomic<bool> bShouldContinueStreaming(false);	
 #if SINE_DEBUG
 	Audio::FSineOsc SineOsc;
 #endif
@@ -53,10 +50,12 @@ namespace PrivateFliteSpeechStreaming
 	{
 		if (!bShouldContinueStreaming)
 		{
-			UE_LOG(LogTextToSpeech, VeryVerbose, TEXT("Stopping speech streaming."));
+			UE_LOG(LogTextToSpeech, VeryVerbose, TEXT("Stopping Flite speech streaming."));
 			return CST_AUDIO_STREAM_STOP;
 		}
-		PCMToFloatBuffer.Reset();
+		
+		// Scratch buffer to convert PCM data for FLite to UE float format 
+		TArray<float> PCMToFloatBuffer;
 		PCMToFloatBuffer.AddZeroed(NumSamples);
 		// Note it's possible for a stream of size 0 to come in with it being the last chunk
 		for (int32 CurrentIndex = 0; CurrentIndex < NumSamples; ++CurrentIndex)
@@ -70,10 +69,8 @@ namespace PrivateFliteSpeechStreaming
 			PCMToFloatBuffer[CurrentIndex] = Sample;
 #endif
 		}
-		int32 SampleRate = static_cast<int32>(Wave->sample_rate);
-		int32 NumChannels = static_cast<int32>(Wave->num_channels);
-		FFliteSynthesizedSpeechData SpeechData(MoveTemp(PCMToFloatBuffer), SampleRate, NumChannels);;
-		int ReturnCode = CST_AUDIO_STREAM_CONT;;
+		FFliteSynthesizedSpeechData SpeechData(MoveTemp(PCMToFloatBuffer), static_cast<int32>(Wave->sample_rate), static_cast<int32>(Wave->num_channels));
+		int ReturnCode = CST_AUDIO_STREAM_CONT;
 		if (bIsLastChunk == 1)
 		{
 			UE_LOG(LogTextToSpeech, VeryVerbose, TEXT("Synthesizing last speech chunk of %i bytes."), SpeechData.GetNumSpeechSamples());
diff --git a/Engine/Plugins/Experimental/TextToSpeech/Source/TextToSpeech/Private/Flite/FliteTextToSpeechSubmixListener.cpp b/Engine/Plugins/Experimental/TextToSpeech/Source/TextToSpeech/Private/Flite/FliteTextToSpeechSubmixListener.cpp
--- a/Engine/Plugins/Experimental/TextToSpeech/Source/TextToSpeech/Private/Flite/FliteTextToSpeechSubmixListener.cpp
+++ b/Engine/Plugins/Experimental/TextToSpeech/Source/TextToSpeech/Private/Flite/FliteTextToSpeechSubmixListener.cpp
@@ -108,10 +108,10 @@ void FFliteTextToSpeechSubmixListener::QueueSynthesizedSpeechChunk_AnyThread(FFl
 			// If pitching down, this buffer will be larger than input buffer
 			// If pitching up, this buffer will be smaller than input buffer
 			int32 NumConvertedSamples = InSynthesizedSpeechChunk.SampleRate / SampleRateConversionRatio;
-			SampleRateConversionBuffer.Reset();
+			int32 OutputSamples = INDEX_NONE;
+			TArray<float> SampleRateConversionBuffer;
 			SampleRateConversionBuffer.AddZeroed(NumConvertedSamples);
 			// Perform the sample rate conversion
-			int32 OutputSamples = INDEX_NONE;
 			int32 ErrorCode = AudioResampler.ProcessAudio(InSynthesizedSpeechChunk.GetSpeechData(), InSynthesizedSpeechChunk.GetNumSpeechSamples(), InSynthesizedSpeechChunk.IsLastChunk(), SampleRateConversionBuffer.GetData(), SampleRateConversionBuffer.Num(), OutputSamples);
 			ensure(OutputSamples <= NumConvertedSamples);
 			if (ErrorCode != 0)
@@ -134,7 +134,7 @@ void FFliteTextToSpeechSubmixListener::QueueSynthesizedSpeechChunk_AnyThread(FFl
 	// we should always have the main audio device. If we get here, something's really wrong
 	ensure(false);
 }
-
+	
 void FFliteTextToSpeechSubmixListener::StopPlayback_GameThread()
 {
 	check(IsInGameThread());
@@ -179,9 +179,17 @@ void FFliteTextToSpeechSubmixListener::Unmute()
 	bMuted = false;
 }
 
+/**
+* Controls the capacity of the FSynthesizedSpeechBuffer. The larger the capacity of the buffer, the more synthesiazed audio data that can be held by 
+*/
+static int32 gFliteSpeechBufferCapacity = 512;
+static FAutoConsoleVariableRef FliteSpeechBufferCapacityRef(
+	TEXT("TextToSpeech.Flite.SpeechBufferCapacity"),
+	gFliteSpeechBufferCapacity,
+	TEXT("Controls the capacity of the Flite synthesized speechbuffer upon initialization of the Flite text to speech submix listener. The buffer will hold chunks of streamed, synthesized speech data. The larger the size the capacity of the buffer, the more audio data can be synthesized at once for playback and the more memory the buffer will take up. It's recommeneded to keep the value greater than 256 and less than 1024.")
+);
 FFliteTextToSpeechSubmixListener::FSynthesizedSpeechBuffer::FSynthesizedSpeechBuffer()
-	// @TODOAccessibility: Don't hardcode the size of the buffer. Also need to figure out a way to conserve space. 
-	: SynthesizedSpeechChunks(256)
+	: SynthesizedSpeechChunks(gFliteSpeechBufferCapacity)
 	, ReadIndex(0)
 	, WriteIndex(0)
 , ChunkReadIndex(0)
diff --git a/Engine/Plugins/Experimental/TextToSpeech/Source/TextToSpeech/Private/Flite/FliteTextToSpeechSubmixListener.h b/Engine/Plugins/Experimental/TextToSpeech/Source/TextToSpeech/Private/Flite/FliteTextToSpeechSubmixListener.h
--- a/Engine/Plugins/Experimental/TextToSpeech/Source/TextToSpeech/Private/Flite/FliteTextToSpeechSubmixListener.h
+++ b/Engine/Plugins/Experimental/TextToSpeech/Source/TextToSpeech/Private/Flite/FliteTextToSpeechSubmixListener.h
@@ -75,12 +75,6 @@ private:
 	Audio::FResampler AudioResampler;
 	FCriticalSection SynthesizedSpeechBufferCS;
 	/**
-	* scratch buffer for SRC
-	* should only be used in background thread when pushing data to
-	* SynthesizedSpeechBuffer
-	*/
-	TArray<float> SampleRateConversionBuffer;
-	/**
 	* The buffer that stores the data with DSP applied
 	* should only be used in audio render thread
 	*/
