Fix for an audio memory stomp in FliteTextToSpeech.
- Introduced audio fences when registering and unregistering FFliteTextToSpeechSubmixListener with the audio device manager.
- Added virtual destructors to audio class interfaces.

diff --git a/Engine/Plugins/Experimental/TextToSpeech/Source/TextToSpeech/Private/Flite/FliteTextToSpeech.cpp b/Engine/Plugins/Experimental/TextToSpeech/Source/TextToSpeech/Private/Flite/FliteTextToSpeech.cpp
--- a/Engine/Plugins/Experimental/TextToSpeech/Source/TextToSpeech/Private/Flite/FliteTextToSpeech.cpp
+++ b/Engine/Plugins/Experimental/TextToSpeech/Source/TextToSpeech/Private/Flite/FliteTextToSpeech.cpp
@@ -141,6 +141,11 @@ void FFliteTextToSpeech::OnActivated()
 			FliteAdapter = MakeUnique<FFliteAdapter>();
 			TTSSubmixListener = MakeUnique<FFliteTextToSpeechSubmixListener>(GetId());
 			AudioDeviceHandle->RegisterSubmixBufferListener(TTSSubmixListener.Get());
+			// RegisterSubmixBufferListener lazily enqueues the registration on the audio thread,
+		// so we have to wait for the audio thread to destroy.
+			FAudioCommandFence Fence;
+			Fence.BeginFence();
+			Fence.Wait();
 		}
 	}
 }
@@ -162,6 +167,9 @@ void FFliteTextToSpeech::OnDeactivated()
 		if (FAudioDeviceHandle AudioDeviceHandle = AudioDeviceManager->GetMainAudioDeviceHandle())
 		{
 			AudioDeviceHandle->UnregisterSubmixBufferListener(TTSSubmixListener.Get());
+			FAudioCommandFence Fence;
+			Fence.BeginFence();
+			Fence.Wait();
 		}
 		FliteAdapter = nullptr;
 		TTSSubmixListener = nullptr;
diff --git a/Engine/Source/Runtime/AudioMixerCore/Public/AudioMixer.h b/Engine/Source/Runtime/AudioMixerCore/Public/AudioMixer.h
--- a/Engine/Source/Runtime/AudioMixerCore/Public/AudioMixer.h
+++ b/Engine/Source/Runtime/AudioMixerCore/Public/AudioMixer.h
@@ -398,6 +398,8 @@ namespace Audio
 	class AUDIOMIXERCORE_API IAudioMixerDeviceChangedListener
 	{
 	public:
+		virtual ~IAudioMixerDeviceChangedListener() = default;
+
 		struct FFormatChangedData
 		{
 			int32 NumChannels = 0;
diff --git a/Engine/Source/Runtime/Engine/Public/ISubmixBufferListener.h b/Engine/Source/Runtime/Engine/Public/ISubmixBufferListener.h
--- a/Engine/Source/Runtime/Engine/Public/ISubmixBufferListener.h
+++ b/Engine/Source/Runtime/Engine/Public/ISubmixBufferListener.h
@@ -7,6 +7,8 @@
 class ENGINE_API ISubmixBufferListener
 {
 public:
+	virtual ~ISubmixBufferListener() = default;
+	
 	/**
 	Called when a new buffer has been rendered for a given submix
 	@param OwningSubmix	The submix object which has rendered a new buffer
