TTS: Unifying the IOS and Mac TTS implementation.

WIth MacOS 14.0 the NSSpeech is deprecated in favor of AVSpeechSynthesizer, which IOS was already using.

The implementation are largely the same, so move the code to the Apple folder and make it work for both platforms.

diff --git a/Engine/Plugins/Experimental/TextToSpeech/Source/TextToSpeech/Private/Apple/AppleTextToSpeech.cpp b/Engine/Plugins/Experimental/TextToSpeech/Source/TextToSpeech/Private/Apple/AppleTextToSpeech.cpp
new file mode 100644
--- /dev/null
+++ b/Engine/Plugins/Experimental/TextToSpeech/Source/TextToSpeech/Private/Apple/AppleTextToSpeech.cpp
@@ -0,0 +1,207 @@
+// Copyright Epic Games, Inc. All Rights Reserved.
+#include "Apple/AppleTextToSpeech.h"
+#include "GenericPlatform/TextToSpeechBase.h"
+#include "TextToSpeechLog.h"
+
+#if PLATFORM_MAC
+#include "Mac/CocoaThread.h"
+#import <AppKit/AppKit.h>
+#endif
+
+#if PLATFORM_IOS
+#include "Async/TaskGraphInterfaces.h"
+#import <UIKit/UIKit.h>
+#endif
+
+// #import <Foundation/Foundation.h>
+#import <AVFoundation/AVFoundation.h>
+
+/** Defined in FTextToSpeechBase.cpp */
+extern TMap<TextToSpeechId, TWeakPtr<FTextToSpeechBase>> ActiveTextToSpeechMap;
+
+@interface FSpeechSynthesizerDelegate : FApplePlatformObject<AVSpeechSynthesizerDelegate>
+-(id)initWithOwningTextToSpeechId:(TextToSpeechId)InOwningTextToSpeechId;
+-(void)dealloc;
+-(void)speechSynthesizer:(AVSpeechSynthesizer*)sender didFinishSpeechUtterance:(AVSpeechUtterance*)utterance;
+@end
+
+@implementation FSpeechSynthesizerDelegate
+{
+	TextToSpeechId _OwningTextToSpeechId;
+}
+APPLE_PLATFORM_OBJECT_ALLOC_OVERRIDES(FSpeechSynthesizerDelegate)
+-(id)initWithOwningTextToSpeechId:(TextToSpeechId)InOwningTextToSpeechId
+{
+	if (self = [super init])
+	{
+		_OwningTextToSpeechId = InOwningTextToSpeechId;
+	}
+	return self;
+}
+
+-(void)dealloc
+{
+	_OwningTextToSpeechId = FTextToSpeechBase::InvalidTextToSpeechId;
+	[super dealloc];
+}
+
+- (void)speechSynthesizer:(AVSpeechSynthesizer*)sender didFinishSpeechUtterance:(AVSpeechUtterance*)utterance
+{
+	// The announcement was completed successfully, not interrupted or manually stopped
+	if (utterance && _OwningTextToSpeechId != FTextToSpeechBase::InvalidTextToSpeechId)
+	{
+		TextToSpeechId IdCopy = _OwningTextToSpeechId;
+		const TMap<TextToSpeechId, TWeakPtr<FTextToSpeechBase>>& ActiveTextToSpeechMapRef = ActiveTextToSpeechMap;
+		GameThreadCall(^{
+			const TWeakPtr<FTextToSpeechBase>* TTSPtr = ActiveTextToSpeechMapRef.Find(IdCopy);
+			if (TTSPtr && TTSPtr->IsValid())
+			{
+				TTSPtr->Pin()->OnTextToSpeechFinishSpeaking_GameThread();
+			}
+		}, false);
+	}
+}
+
+@end
+
+FAppleTextToSpeech::FAppleTextToSpeech()
+	: bIsSpeaking(false)
+	, Volume(0.0f)
+	, Rate(0.0f)
+	, SpeechSynthesizer(nullptr)
+	, SpeechSynthesizerDelegate(nullptr)
+{
+
+}
+
+FAppleTextToSpeech::~FAppleTextToSpeech()
+{
+	// base class already takes care of checking if the TTS is active and calls deactivate
+}
+
+void FAppleTextToSpeech::Speak(const FString& InStringToSpeak)
+{
+	if (IsActive())
+	{
+		UE_LOG(LogTextToSpeech, Verbose, TEXT("Apple TTS speak requested."));
+		if (!InStringToSpeak.IsEmpty())
+		{
+			UE_LOG(LogTextToSpeech, VeryVerbose, TEXT("String to speak: %s"), *InStringToSpeak);
+			if (IsSpeaking())
+			{
+				StopSpeaking();
+			}
+			SCOPED_AUTORELEASE_POOL;
+			NSString* Announcement = InStringToSpeak.GetNSString();
+			ensureMsgf(Announcement, TEXT("Failed to convert FString to NSString for TTS."));
+			AVSpeechUtterance* Utterance = [AVSpeechUtterance speechUtteranceWithString:Announcement];
+			checkf(Utterance, TEXT("Failed to create an utterance for text to speech."));
+			// for now we just use the default system language that's being used
+			Utterance.voice = [AVSpeechSynthesisVoice voiceWithLanguage: [AVSpeechSynthesisVoice currentLanguageCode]];
+			// If muted, set to volume to 0
+			Utterance.volume = IsMuted() ? 0.0f : Volume;
+			Utterance.rate = Rate;
+			[SpeechSynthesizer speakUtterance:Utterance];
+			bIsSpeaking = true;
+		}
+	}
+}
+
+bool FAppleTextToSpeech::IsSpeaking() const
+{
+	return IsActive() ? bIsSpeaking : false;
+}
+
+void FAppleTextToSpeech::StopSpeaking()
+{
+	if (IsActive())
+	{
+		if (IsSpeaking())
+		{
+			[SpeechSynthesizer stopSpeakingAtBoundary:AVSpeechBoundaryImmediate];
+		}
+		bIsSpeaking = false;
+		UE_LOG(LogTextToSpeech, Verbose, TEXT("Apple TTS stopped speaking."));
+	}
+}
+
+float FAppleTextToSpeech::GetVolume() const
+{
+	return Volume;
+}
+
+void FAppleTextToSpeech::SetVolume(float InVolume)
+{
+	Volume = FMath::Clamp(InVolume, 0.f, 1.0f);
+}
+
+float FAppleTextToSpeech::GetRate() const
+{
+	return Rate;
+}
+
+void FAppleTextToSpeech::SetRate(float InRate)
+{
+	Rate = FMath::Clamp(InRate, 0.0f, 1.0f);
+}
+
+void FAppleTextToSpeech::Mute()
+{
+	if (IsActive() && !IsMuted())
+	{
+		SetMuted(true);
+	}
+}
+
+void FAppleTextToSpeech::Unmute()
+{
+	if (IsActive() && IsMuted())
+	{
+		SetMuted(false);
+	}
+}
+
+void FAppleTextToSpeech::OnActivated()
+{
+	ensureMsgf(!IsActive(), TEXT("Attempting to activate an already activated TTS. FTextToSpeechBase::Activate() should already guard against this."));
+	TextToSpeechId IdCopy = GetId();
+	MainThreadCall(^{
+		SpeechSynthesizer = [[AVSpeechSynthesizer alloc] init];
+		SpeechSynthesizerDelegate = [[FSpeechSynthesizerDelegate alloc] initWithOwningTextToSpeechId:IdCopy];
+		SpeechSynthesizer.delegate = SpeechSynthesizerDelegate;
+
+		// To get initial volume and rate, we need to retrieve it from an utterance
+		AVSpeechUtterance* Utterance = [AVSpeechUtterance speechUtteranceWithString:@"Temp"];
+		checkf(Utterance, TEXT("Failed to create an utterance for TTS."));
+		Volume = Utterance.volume;
+		Rate = Utterance.rate;
+		ensure(Volume >= 0 && Volume <= 1.0f);
+		ensure(Rate >= 0 && Rate <= 1.0f);
+
+#if PLATFORM_IOS
+		// This allows us to still hear the TTS when the IOS ringer is muted
+		// SetFeature already takes care of keeping track of how many requests are made to activate/deactivate an audio feature
+		[[IOSAppDelegate GetDelegate] SetFeature:EAudioFeature::Playback Active:true];
+#endif
+	});
+
+	UE_LOG(LogTextToSpeech, Verbose, TEXT("Apple TTS activated."));
+}
+
+void FAppleTextToSpeech::OnDeactivated()
+{
+	ensureMsgf(IsActive(), TEXT("Attempting to deactivate an already deactivated TTS. FTextToSpeechBase::Deactivate() should already guard against this."));
+	// deallocate all AppKit objects in main thread just in case
+	MainThreadCall(^{
+		checkf(SpeechSynthesizerDelegate, TEXT("Deactivating Apple TTS with null speech synthesizer delegate. Speech synthesizer delegate must be valid throughout the lifetime of the object."));
+		[SpeechSynthesizerDelegate release];
+		checkf(SpeechSynthesizer, TEXT("Deactivating Apple TTS with null speech synthesizer. Speech synthesizer must be valid throughout the lifetime of the object."));
+		[SpeechSynthesizer release];
+	});
+	UE_LOG(LogTextToSpeech, Verbose, TEXT("Apple TTS deactivated."));
+
+#if PLATFORM_IOS
+		// SetFeature already takes care of keeping track of the number of requests to activate/deactivate a feature
+		[[IOSAppDelegate GetDelegate] SetFeature:EAudioFeature::Playback Active:false];
+#endif
+}
\ No newline at end of file
diff --git a/Engine/Plugins/Experimental/TextToSpeech/Source/TextToSpeech/Private/Apple/AppleTextToSpeech.h b/Engine/Plugins/Experimental/TextToSpeech/Source/TextToSpeech/Private/Apple/AppleTextToSpeech.h
new file mode 100644
--- /dev/null
+++ b/Engine/Plugins/Experimental/TextToSpeech/Source/TextToSpeech/Private/Apple/AppleTextToSpeech.h
@@ -0,0 +1,51 @@
+// Copyright Epic Games, Inc. All Rights Reserved.
+
+#pragma once
+#include "CoreMinimal.h"
+#include "GenericPlatform/TextToSpeechBase.h"
+
+@class AVSpeechSynthesizer;
+@class FSpeechSynthesizerDelegate;
+/**
+* The text to speech implementation for Mac
+*/
+class 	FAppleTextToSpeech : public FTextToSpeechBase
+{
+public:
+	FAppleTextToSpeech();
+	virtual ~FAppleTextToSpeech();
+	
+	// FGenericTextToSpeech
+	virtual void Speak(const FString& InStringToSpeak) override final;
+	virtual bool IsSpeaking() const override final;
+	virtual void StopSpeaking() override final;
+	virtual float GetVolume() const override final;
+	virtual void SetVolume(float InVolume) override final;
+	virtual float GetRate() const override final;
+	virtual void SetRate(float InRate) override final;
+	virtual void Mute() override final;
+	virtual void Unmute() override final;
+protected:
+	virtual void OnActivated() override final;
+	virtual void OnDeactivated() override final;
+	// ~
+private:
+	/**
+	* True if the speech synthesizer is currently synthesizing any text, else false.
+	* Speech synthesis for Mac is asynchronous, but we make the checking of whether something is speaking synchrnous
+	* so as to simplify algorithms that depend on the check.
+	* A precise check is not currently necessary.
+	 */
+ 	bool bIsSpeaking;
+	/** The current volume of the speech synthesizer. */
+	float Volume;
+	/** The current speech rate of the speech synthesizer */
+	float Rate;
+	/** The platform speech synthesizer that converts text to speech */
+	AVSpeechSynthesizer* SpeechSynthesizer;
+	/**
+	* The delegate for the speech synthesizer. Callbacks indicating
+	* speech synthesis playback progress, intruuption and completion are handled by this delegate.
+	 */
+	FSpeechSynthesizerDelegate* SpeechSynthesizerDelegate;
+};
\ No newline at end of file
diff --git a/Engine/Plugins/Experimental/TextToSpeech/Source/TextToSpeech/Private/Apple/AppleTextToSpeechFactory.cpp b/Engine/Plugins/Experimental/TextToSpeech/Source/TextToSpeech/Private/Apple/AppleTextToSpeechFactory.cpp
new file mode 100644
--- /dev/null
+++ b/Engine/Plugins/Experimental/TextToSpeech/Source/TextToSpeech/Private/Apple/AppleTextToSpeechFactory.cpp
@@ -0,0 +1,9 @@
+// Copyright Epic Games, Inc. All Rights Reserved.
+#include "Apple/AppleTextToSpeechFactory.h"
+#include "GenericPlatform/ITextToSpeechFactory.h"
+#include "Apple/AppleTextToSpeech.h"
+
+TSharedRef<FTextToSpeechBase> FAppleTextToSpeechFactory::Create()
+{
+	return MakeShared<FAppleTextToSpeech>();
+}
diff --git a/Engine/Plugins/Experimental/TextToSpeech/Source/TextToSpeech/Private/Apple/AppleTextToSpeechFactory.h b/Engine/Plugins/Experimental/TextToSpeech/Source/TextToSpeech/Private/Apple/AppleTextToSpeechFactory.h
new file mode 100644
--- /dev/null
+++ b/Engine/Plugins/Experimental/TextToSpeech/Source/TextToSpeech/Private/Apple/AppleTextToSpeechFactory.h
@@ -0,0 +1,18 @@
+// Copyright Epic Games, Inc. All Rights Reserved.
+#pragma once
+#include "CoreMinimal.h"
+#include "GenericPlatform/ITextToSpeechFactory.h"
+
+class FTextToSpeechBase;
+/** The platform default text to speech factory for Apple. */
+class FAppleTextToSpeechFactory : public ITextToSpeechFactory
+{
+public:
+	FAppleTextToSpeechFactory() = default;
+	virtual ~FAppleTextToSpeechFactory() = default;
+
+	// ITextToSpeechFactory
+	virtual TSharedRef<FTextToSpeechBase> Create() override;
+	// ~
+};
+typedef FAppleTextToSpeechFactory FPlatformTextToSpeechFactory;
diff --git a/Engine/Plugins/Experimental/TextToSpeech/Source/TextToSpeech/Private/IOS/IOSTextToSpeech.cpp b/Engine/Plugins/Experimental/TextToSpeech/Source/TextToSpeech/Private/IOS/IOSTextToSpeech.cpp
deleted file mode 100644
diff --git a/Engine/Plugins/Experimental/TextToSpeech/Source/TextToSpeech/Private/IOS/IOSTextToSpeech.h b/Engine/Plugins/Experimental/TextToSpeech/Source/TextToSpeech/Private/IOS/IOSTextToSpeech.h
deleted file mode 100644
diff --git a/Engine/Plugins/Experimental/TextToSpeech/Source/TextToSpeech/Private/IOS/IOSTextToSpeechFactory.cpp b/Engine/Plugins/Experimental/TextToSpeech/Source/TextToSpeech/Private/IOS/IOSTextToSpeechFactory.cpp
deleted file mode 100644
diff --git a/Engine/Plugins/Experimental/TextToSpeech/Source/TextToSpeech/Private/IOS/IOSTextToSpeechFactory.h b/Engine/Plugins/Experimental/TextToSpeech/Source/TextToSpeech/Private/IOS/IOSTextToSpeechFactory.h
--- a/Engine/Plugins/Experimental/TextToSpeech/Source/TextToSpeech/Private/IOS/IOSTextToSpeechFactory.h
+++ b/Engine/Plugins/Experimental/TextToSpeech/Source/TextToSpeech/Private/IOS/IOSTextToSpeechFactory.h
@@ -1,20 +1,4 @@
-// Copyright Epic Games, Inc. All Rights Reserved.
-
 #pragma once
-#if PLATFORM_IOS
-#include "CoreMinimal.h"
-#include "GenericPlatform/ITextToSpeechFactory.h"
-
-class FTextToSpeechBase;
-/** The platform default text to speech factory for IOS. */
-class FIOSTextToSpeechFactory : public ITextToSpeechFactory
-{
-public:
-	FIOSTextToSpeechFactory() = default;
-	virtual ~FIOSTextToSpeechFactory() = default;
-
-	virtual TSharedRef<FTextToSpeechBase> Create() override;
-};
-typedef FIOSTextToSpeechFactory FPlatformTextToSpeechFactory;
-#endif
-
+// We're doing this as the TextToSpeechModule.cpp is using COMPILED_PLATFORM_HEADER
+// which doesn't have a clean way to override and point to a shared platform.
+#include "Apple/AppleTextToSpeechFactory.h"
\ No newline at end of file
diff --git a/Engine/Plugins/Experimental/TextToSpeech/Source/TextToSpeech/Private/Mac/MacTextToSpeech.cpp b/Engine/Plugins/Experimental/TextToSpeech/Source/TextToSpeech/Private/Mac/MacTextToSpeech.cpp
deleted file mode 100644
diff --git a/Engine/Plugins/Experimental/TextToSpeech/Source/TextToSpeech/Private/Mac/MacTextToSpeech.h b/Engine/Plugins/Experimental/TextToSpeech/Source/TextToSpeech/Private/Mac/MacTextToSpeech.h
deleted file mode 100644
diff --git a/Engine/Plugins/Experimental/TextToSpeech/Source/TextToSpeech/Private/Mac/MacTextToSpeechFactory.cpp b/Engine/Plugins/Experimental/TextToSpeech/Source/TextToSpeech/Private/Mac/MacTextToSpeechFactory.cpp
deleted file mode 100644
diff --git a/Engine/Plugins/Experimental/TextToSpeech/Source/TextToSpeech/Private/Mac/MacTextToSpeechFactory.h b/Engine/Plugins/Experimental/TextToSpeech/Source/TextToSpeech/Private/Mac/MacTextToSpeechFactory.h
--- a/Engine/Plugins/Experimental/TextToSpeech/Source/TextToSpeech/Private/Mac/MacTextToSpeechFactory.h
+++ b/Engine/Plugins/Experimental/TextToSpeech/Source/TextToSpeech/Private/Mac/MacTextToSpeechFactory.h
@@ -1,22 +1,4 @@
-// Copyright Epic Games, Inc. All Rights Reserved.
-
 #pragma once
-#if PLATFORM_MAC
-#include "CoreMinimal.h"
-#include "GenericPlatform/ITextToSpeechFactory.h"
-
-class FTextToSpeechBase;
-/** The platform default text to speech factory for Mac. */
-class FMacTextToSpeechFactory : public ITextToSpeechFactory
-{
-public:
-	FMacTextToSpeechFactory() = default;
-	virtual ~FMacTextToSpeechFactory() = default;
-
-	// ITextToSpeechFactory
-	virtual TSharedRef<FTextToSpeechBase> Create() override;
-	// ~
-};
-typedef FMacTextToSpeechFactory FPlatformTextToSpeechFactory;
-#endif
-
+// We're doing this as the TextToSpeechModule.cpp is using COMPILED_PLATFORM_HEADER
+// which doesn't have a clean way to override and point to a shared platform.
+#include "Apple/AppleTextToSpeechFactory.h"
\ No newline at end of file
