TTS: Add missing copyrights.

diff --git a/Engine/Plugins/Experimental/TextToSpeech/Source/TextToSpeech/Private/IOS/IOSTextToSpeechFactory.h b/Engine/Plugins/Experimental/TextToSpeech/Source/TextToSpeech/Private/IOS/IOSTextToSpeechFactory.h
--- a/Engine/Plugins/Experimental/TextToSpeech/Source/TextToSpeech/Private/IOS/IOSTextToSpeechFactory.h
+++ b/Engine/Plugins/Experimental/TextToSpeech/Source/TextToSpeech/Private/IOS/IOSTextToSpeechFactory.h
@@ -1,3 +1,4 @@
+// Copyright Epic Games, Inc. All Rights Reserved.
 #pragma once
 // We're doing this as the TextToSpeechModule.cpp is using COMPILED_PLATFORM_HEADER
 // which doesn't have a clean way to override and point to a shared platform.
diff --git a/Engine/Plugins/Experimental/TextToSpeech/Source/TextToSpeech/Private/Mac/MacTextToSpeechFactory.h b/Engine/Plugins/Experimental/TextToSpeech/Source/TextToSpeech/Private/Mac/MacTextToSpeechFactory.h
--- a/Engine/Plugins/Experimental/TextToSpeech/Source/TextToSpeech/Private/Mac/MacTextToSpeechFactory.h
+++ b/Engine/Plugins/Experimental/TextToSpeech/Source/TextToSpeech/Private/Mac/MacTextToSpeechFactory.h
@@ -1,3 +1,4 @@
+// Copyright Epic Games, Inc. All Rights Reserved.
 #pragma once
 // We're doing this as the TextToSpeechModule.cpp is using COMPILED_PLATFORM_HEADER
 // which doesn't have a clean way to override and point to a shared platform.
