TTS: Fixing up iOS issues.

diff --git a/Engine/Plugins/Experimental/TextToSpeech/Source/TextToSpeech/Private/Apple/AppleTextToSpeech.cpp b/Engine/Plugins/Experimental/TextToSpeech/Source/TextToSpeech/Private/Apple/AppleTextToSpeech.cpp
--- a/Engine/Plugins/Experimental/TextToSpeech/Source/TextToSpeech/Private/Apple/AppleTextToSpeech.cpp
+++ b/Engine/Plugins/Experimental/TextToSpeech/Source/TextToSpeech/Private/Apple/AppleTextToSpeech.cpp
@@ -10,6 +10,7 @@
 
 #if PLATFORM_IOS
 #include "Async/TaskGraphInterfaces.h"
+#include "IOS/IOSAppDelegate.h"
 #import <UIKit/UIKit.h>
 #endif
 
@@ -52,6 +53,7 @@ APPLE_PLATFORM_OBJECT_ALLOC_OVERRIDES(FSpeechSynthesizerDelegate)
 	{
 		TextToSpeechId IdCopy = _OwningTextToSpeechId;
 		const TMap<TextToSpeechId, TWeakPtr<FTextToSpeechBase>>& ActiveTextToSpeechMapRef = ActiveTextToSpeechMap;
+#if PLATFORM_MAC
 		GameThreadCall(^{
 			const TWeakPtr<FTextToSpeechBase>* TTSPtr = ActiveTextToSpeechMapRef.Find(IdCopy);
 			if (TTSPtr && TTSPtr->IsValid())
@@ -59,6 +61,16 @@ APPLE_PLATFORM_OBJECT_ALLOC_OVERRIDES(FSpeechSynthesizerDelegate)
 				TTSPtr->Pin()->OnTextToSpeechFinishSpeaking_GameThread();
 			}
 		}, false);
+#else // PLATFORM_IOS
+		FFunctionGraphTask::CreateAndDispatchWhenReady([IdCopy, ActiveTextToSpeechMapRef]()
+		{
+			const TWeakPtr<FTextToSpeechBase>* TTSPtr = ActiveTextToSpeechMapRef.Find(IdCopy);
+			if (TTSPtr && TTSPtr->IsValid())
+			{
+				TTSPtr->Pin()->OnTextToSpeechFinishSpeaking_GameThread();
+			}
+		}, TStatId(), NULL, ENamedThreads::GameThread);
+#endif
 	}
 }
 
