TTS: Wrapping other maincall correctly for iOS

diff --git a/Engine/Plugins/Experimental/TextToSpeech/Source/TextToSpeech/Private/Apple/AppleTextToSpeech.cpp b/Engine/Plugins/Experimental/TextToSpeech/Source/TextToSpeech/Private/Apple/AppleTextToSpeech.cpp
--- a/Engine/Plugins/Experimental/TextToSpeech/Source/TextToSpeech/Private/Apple/AppleTextToSpeech.cpp
+++ b/Engine/Plugins/Experimental/TextToSpeech/Source/TextToSpeech/Private/Apple/AppleTextToSpeech.cpp
@@ -177,7 +177,12 @@ void FAppleTextToSpeech::OnActivated()
 {
 	ensureMsgf(!IsActive(), TEXT("Attempting to activate an already activated TTS. FTextToSpeechBase::Activate() should already guard against this."));
 	TextToSpeechId IdCopy = GetId();
+#if PLATFORM_MAC
 	MainThreadCall(^{
+#else // PLATFORM_IOS
+	dispatch_async(dispatch_get_main_queue(), ^
+	{
+#endif
 		SpeechSynthesizer = [[AVSpeechSynthesizer alloc] init];
 		SpeechSynthesizerDelegate = [[FSpeechSynthesizerDelegate alloc] initWithOwningTextToSpeechId:IdCopy];
 		SpeechSynthesizer.delegate = SpeechSynthesizerDelegate;
@@ -204,7 +209,12 @@ void FAppleTextToSpeech::OnDeactivated()
 {
 	ensureMsgf(IsActive(), TEXT("Attempting to deactivate an already deactivated TTS. FTextToSpeechBase::Deactivate() should already guard against this."));
 	// deallocate all AppKit objects in main thread just in case
+#if PLATFORM_MAC
 	MainThreadCall(^{
+#else // PLATFORM_IOS
+	dispatch_async(dispatch_get_main_queue(), ^
+	{
+#endif
 		checkf(SpeechSynthesizerDelegate, TEXT("Deactivating Apple TTS with null speech synthesizer delegate. Speech synthesizer delegate must be valid throughout the lifetime of the object."));
 		[SpeechSynthesizerDelegate release];
 		checkf(SpeechSynthesizer, TEXT("Deactivating Apple TTS with null speech synthesizer. Speech synthesizer must be valid throughout the lifetime of the object."));
